- Fixing rcwebinterface non unity build

#rnx

diff --git a/Engine/Plugins/VirtualProduction/RemoteControlWebInterface/Source/RemoteControlWebInterface/Private/RCWebInterface.cpp b/Engine/Plugins/VirtualProduction/RemoteControlWebInterface/Source/RemoteControlWebInterface/Private/RCWebInterface.cpp
--- a/Engine/Plugins/VirtualProduction/RemoteControlWebInterface/Source/RemoteControlWebInterface/Private/RCWebInterface.cpp
+++ b/Engine/Plugins/VirtualProduction/RemoteControlWebInterface/Source/RemoteControlWebInterface/Private/RCWebInterface.cpp
@@ -12,7 +12,6 @@
 
 #if WITH_EDITOR
 #include "Editor.h"
-#include "Misc/CoreDelegates.h"
 #include "RCWebInterfaceCustomizations.h"
 #endif
 
