- Adding more logging capabilities to diagnostic poor performance with webapp

diff --git a/Engine/Plugins/Experimental/WebSocketNetworking/Source/WebSocketNetworking/Private/WebSocketNetworkingPlugin.cpp b/Engine/Plugins/Experimental/WebSocketNetworking/Source/WebSocketNetworking/Private/WebSocketNetworkingPlugin.cpp
--- a/Engine/Plugins/Experimental/WebSocketNetworking/Source/WebSocketNetworking/Private/WebSocketNetworkingPlugin.cpp
+++ b/Engine/Plugins/Experimental/WebSocketNetworking/Source/WebSocketNetworking/Private/WebSocketNetworkingPlugin.cpp
@@ -3,6 +3,8 @@
 #include "IWebSocketNetworkingModule.h"
 #include "WebSocketNetworkingPrivate.h"
 #include "WebSocketServer.h"
+#include "WebSocket.h"
+
 
 
 class FWebSocketNetworkingPlugin : public IWebSocketNetworkingModule
@@ -23,6 +25,11 @@ public:
 		return MakeUnique<FWebSocketServer>();
 	}
 
+	TSharedPtr<INetworkingWebSocket> CreateConnection(const FInternetAddr& ServerAddress) override
+	{
+		return MakeShared<FWebSocket>(ServerAddress);
+	}
+
 };
 
 IMPLEMENT_MODULE(FWebSocketNetworkingPlugin, WebSocketNetworking)
diff --git a/Engine/Plugins/Experimental/WebSocketNetworking/Source/WebSocketNetworking/Public/IWebSocketNetworkingModule.h b/Engine/Plugins/Experimental/WebSocketNetworking/Source/WebSocketNetworking/Public/IWebSocketNetworkingModule.h
--- a/Engine/Plugins/Experimental/WebSocketNetworking/Source/WebSocketNetworking/Public/IWebSocketNetworkingModule.h
+++ b/Engine/Plugins/Experimental/WebSocketNetworking/Source/WebSocketNetworking/Public/IWebSocketNetworkingModule.h
@@ -4,8 +4,11 @@
 
 #include "Modules/ModuleInterface.h"
 #include "Templates/UniquePtr.h"
+#include "Templates/SharedPointer.h"
+
 
 class IWebSocketServer;
+class INetworkingWebSocket;
 
 /* 
  * Interface for the WebSocketNetworking module. 
@@ -22,4 +25,11 @@ public:
 	 * @return A new WebSocket server, or nullptr if the server couldn't be created.
 	 */
 	virtual TUniquePtr<IWebSocketServer> CreateServer() = 0;
+
+	/**
+	 * Create a WebSocket client connection.
+	 *
+	 * @return A new WebSocket client connection, or nullptr if the connection couldn't be created.
+	 */
+	virtual TSharedPtr<INetworkingWebSocket> CreateConnection(const class FInternetAddr& ServerAddress) = 0;
 };
diff --git a/Engine/Plugins/VirtualProduction/RemoteControl/Source/RemoteControlCommon/Public/RemoteControlSettings.h b/Engine/Plugins/VirtualProduction/RemoteControl/Source/RemoteControlCommon/Public/RemoteControlSettings.h
--- a/Engine/Plugins/VirtualProduction/RemoteControl/Source/RemoteControlCommon/Public/RemoteControlSettings.h
+++ b/Engine/Plugins/VirtualProduction/RemoteControl/Source/RemoteControlCommon/Public/RemoteControlSettings.h
@@ -56,6 +56,10 @@ public:
 	UPROPERTY(config, EditAnywhere, Category = "Remote Control Web Interface", DisplayName = "Force WebApp build at startup")
 	bool bForceWebAppBuildAtStartup = false;
 
+	/** Should WebApp log timing. */
+	UPROPERTY(config, EditAnywhere, Category = "Remote Control Web Interface", DisplayName = "Log WebApp requests handle duration")
+	bool bWebAppLogRequestDuration = false;
+
 	/** Whether web server is started automatically. */
 	UPROPERTY(config, EditAnywhere, Category = "Remote Control Web Server")
 	bool bAutoStartWebServer = true;
diff --git a/Engine/Plugins/VirtualProduction/RemoteControl/Source/WebRemoteControl/Private/WebRemoteControl.cpp b/Engine/Plugins/VirtualProduction/RemoteControl/Source/WebRemoteControl/Private/WebRemoteControl.cpp
--- a/Engine/Plugins/VirtualProduction/RemoteControl/Source/WebRemoteControl/Private/WebRemoteControl.cpp
+++ b/Engine/Plugins/VirtualProduction/RemoteControl/Source/WebRemoteControl/Private/WebRemoteControl.cpp
@@ -11,6 +11,7 @@
 #include "RemoteControlSettings.h"
 #include "RemoteControlPreset.h"
 #include "WebRemoteControlUtils.h"
+#include "WebRemoteControlExternalLogger.h"
 #include "WebSocketMessageHandler.h"
 
 #if WITH_EDITOR
@@ -455,6 +456,16 @@ void FWebRemoteControlModule::StopWebSocketServer()
 	OnWebSocketServerStoppedDelegate.Broadcast();
 }
 
+void FWebRemoteControlModule::SetExternalRemoteWebSocketLoggerConnection(TSharedPtr<INetworkingWebSocket> WebSocketLoggerConnection)
+{
+	ExternalLogger.Reset();
+	
+	if (WebSocketLoggerConnection.IsValid())
+	{
+		ExternalLogger = MakeUnique<FWebRemoteControlExternalLogger>(WebSocketLoggerConnection);
+	}
+}
+
 void FWebRemoteControlModule::StartRoute(const FRemoteControlRoute& Route)
 {
 	// The handler is wrapped in a lambda since HttpRouter::BindRoute only accepts TFunctions
@@ -1697,10 +1708,15 @@ void FWebRemoteControlModule::HandleWebSocketHttpMessage(const FRemoteControlWeb
 		return;
 	}
 
+	LogRequestExternally(Wrapper.RequestId, TEXT("UE Received"));
+	
 	FMemoryWriter Writer(UTF8Response);
 	InvokeWrappedRequest(Wrapper, Writer);
 
+	LogRequestExternally(Wrapper.RequestId, TEXT("UE Processed"));
+
 	WebSocketServer.Send(WebSocketMessage.ClientId, MoveTemp(UTF8Response));
+	LogRequestExternally(Wrapper.RequestId, TEXT("UE Sent"));
 }
 
 void FWebRemoteControlModule::InvokeWrappedRequest(const FRCRequestWrapper& Wrapper, FMemoryWriter& OutUTF8PayloadWriter, const FHttpServerRequest* TemplateRequest)
@@ -1760,6 +1776,14 @@ void FWebRemoteControlModule::OnSettingsModified(UObject* Settings, FPropertyCha
 
 #endif
 
+void FWebRemoteControlModule::LogRequestExternally(int32 RequestId, const TCHAR* Stage)
+{
+	if (ExternalLogger.IsValid())
+	{
+		ExternalLogger->Log(RequestId, Stage);
+	}
+}
+
 #undef LOCTEXT_NAMESPACE /* WebRemoteControl */
 
 IMPLEMENT_MODULE(FWebRemoteControlModule, WebRemoteControl);
diff --git a/Engine/Plugins/VirtualProduction/RemoteControl/Source/WebRemoteControl/Private/WebRemoteControl.h b/Engine/Plugins/VirtualProduction/RemoteControl/Source/WebRemoteControl/Private/WebRemoteControl.h
--- a/Engine/Plugins/VirtualProduction/RemoteControl/Source/WebRemoteControl/Private/WebRemoteControl.h
+++ b/Engine/Plugins/VirtualProduction/RemoteControl/Source/WebRemoteControl/Private/WebRemoteControl.h
@@ -46,6 +46,7 @@ public:
 	virtual FSimpleMulticastDelegate& OnHttpServerStopped() override { return OnHttpServerStoppedDelegate; }
 	virtual FOnWebServerStarted& OnWebSocketServerStarted() override { return OnWebSocketServerStartedDelegate; }
 	virtual FSimpleMulticastDelegate& OnWebSocketServerStopped() override { return OnWebSocketServerStoppedDelegate; }
+	virtual void SetExternalRemoteWebSocketLoggerConnection(TSharedPtr<class INetworkingWebSocket> WebSocketLoggerConnection);
 	//~ End IWebRemoteControlModule Interface
 
 	/**
@@ -140,6 +141,8 @@ private:
 	void OnSettingsModified(UObject* Settings, FPropertyChangedEvent& PropertyChangedEvent);
 #endif
 
+	void LogRequestExternally(int32 RequestId, const TCHAR* Stage);
+
 private:
 	/** Console commands handles. */
 	TArray<TUniquePtr<FAutoConsoleCommand>> ConsoleCommands;
@@ -185,6 +188,9 @@ private:
 	 */
 	TMap<FDelegateHandle, FDelegateHandle> PreprocessorsHandleMappings;
 
+	/** An external web socket logger*/
+	TUniquePtr<class FWebRemoteControlExternalLogger> ExternalLogger;
+
 	//~ Server started stopped delegates.
 	FOnWebServerStarted OnHttpServerStartedDelegate;
 	FSimpleMulticastDelegate OnHttpServerStoppedDelegate;
diff --git a/Engine/Plugins/VirtualProduction/RemoteControl/Source/WebRemoteControl/Private/WebRemoteControlExternalLogger.cpp b/Engine/Plugins/VirtualProduction/RemoteControl/Source/WebRemoteControl/Private/WebRemoteControlExternalLogger.cpp
new file mode 100644
--- /dev/null
+++ b/Engine/Plugins/VirtualProduction/RemoteControl/Source/WebRemoteControl/Private/WebRemoteControlExternalLogger.cpp
@@ -0,0 +1,51 @@
+// Copyright Epic Games, Inc. All Rights Reserved.
+
+#pragma once
+
+#include "WebRemoteControlExternalLogger.h"
+#include "WebSocketNetworkingDelegates.h"
+#include "WebRemoteControlUtils.h"
+#include "INetworkingWebSocket.h"
+
+
+FWebRemoteControlExternalLogger::FWebRemoteControlExternalLogger(TSharedPtr<INetworkingWebSocket> WebSocketConnectionIn) :
+	WebSocketConnection(WebSocketConnectionIn)
+{
+	TickerHandle = FTSTicker::GetCoreTicker().AddTicker(FTickerDelegate::CreateRaw(this, &FWebRemoteControlExternalLogger::OnTick));
+
+	FWebSocketInfoCallBack ConnectedCallback;
+	ConnectedCallback.BindRaw(this, &FWebRemoteControlExternalLogger::OnConnected);
+	WebSocketConnection->SetConnectedCallBack(ConnectedCallback);
+
+	FWebSocketInfoCallBack ClosedCallback;
+	ClosedCallback.BindRaw(this, &FWebRemoteControlExternalLogger::OnClosed);
+	WebSocketConnection->SetSocketClosedCallBack(ClosedCallback);
+}
+
+void FWebRemoteControlExternalLogger::Log(int32 RequestId, const TCHAR* Stage)
+{
+	if (IsConnected)
+	{
+		const FString Json = FString::Printf(TEXT("{\"RequestId\":%d,\"Frame\":%d,\"UETime\":%f,\"Stage\":\"%s\"}"), RequestId, GFrameNumber, FPlatformTime::Seconds(), Stage);
+
+		TArray<uint8> Buffer;
+		WebRemoteControlUtils::ConvertToUTF8(Json, Buffer);
+		WebSocketConnection->Send(Buffer.GetData(), Buffer.Num(), false);
+	}
+}
+
+bool FWebRemoteControlExternalLogger::OnTick(float DeltaSeconds)
+{
+	WebSocketConnection->Tick();
+	return true;
+}
+
+void FWebRemoteControlExternalLogger::OnConnected()
+{
+	IsConnected = true;
+}
+
+void FWebRemoteControlExternalLogger::OnClosed()
+{
+	IsConnected = false;
+}
diff --git a/Engine/Plugins/VirtualProduction/RemoteControl/Source/WebRemoteControl/Private/WebRemoteControlExternalLogger.h b/Engine/Plugins/VirtualProduction/RemoteControl/Source/WebRemoteControl/Private/WebRemoteControlExternalLogger.h
new file mode 100644
--- /dev/null
+++ b/Engine/Plugins/VirtualProduction/RemoteControl/Source/WebRemoteControl/Private/WebRemoteControlExternalLogger.h
@@ -0,0 +1,28 @@
+// Copyright Epic Games, Inc. All Rights Reserved.
+
+#pragma once
+
+#include "Containers/Ticker.h"
+#include "Templates/SharedPointer.h"
+
+
+class FWebRemoteControlExternalLogger
+{
+public:
+	FWebRemoteControlExternalLogger(TSharedPtr<class INetworkingWebSocket> WebSocketConnection);
+
+	void Log(int32 RequestId, const TCHAR* Stage);
+
+private:
+	bool OnTick(float DeltaSeconds);
+	void OnConnected();
+	void OnClosed();
+
+	bool IsConnected = false;
+
+	/** A WebSocket connection to the logger server*/
+	TSharedPtr<class INetworkingWebSocket> WebSocketConnection;
+
+	/** A handle to the ticker callback */
+	FTSTicker::FDelegateHandle TickerHandle;
+};
\ No newline at end of file
diff --git a/Engine/Plugins/VirtualProduction/RemoteControl/Source/WebRemoteControl/Public/IWebRemoteControlModule.h b/Engine/Plugins/VirtualProduction/RemoteControl/Source/WebRemoteControl/Public/IWebRemoteControlModule.h
--- a/Engine/Plugins/VirtualProduction/RemoteControl/Source/WebRemoteControl/Public/IWebRemoteControlModule.h
+++ b/Engine/Plugins/VirtualProduction/RemoteControl/Source/WebRemoteControl/Public/IWebRemoteControlModule.h
@@ -49,4 +49,9 @@ public:
 	 * Event triggered when the websocket server stops.
 	 */
 	virtual FSimpleMulticastDelegate& OnWebSocketServerStopped() = 0;
+
+	/**
+	 * Set Remote Web Socket Logger
+	 */
+	virtual void SetExternalRemoteWebSocketLoggerConnection(TSharedPtr<class INetworkingWebSocket> WebSocketLoggerConnection) = 0;
 };
diff --git a/Engine/Plugins/VirtualProduction/RemoteControlWebInterface/Source/RemoteControlWebInterface/Private/RCWebInterface.cpp b/Engine/Plugins/VirtualProduction/RemoteControlWebInterface/Source/RemoteControlWebInterface/Private/RCWebInterface.cpp
--- a/Engine/Plugins/VirtualProduction/RemoteControlWebInterface/Source/RemoteControlWebInterface/Private/RCWebInterface.cpp
+++ b/Engine/Plugins/VirtualProduction/RemoteControlWebInterface/Source/RemoteControlWebInterface/Private/RCWebInterface.cpp
@@ -96,6 +96,12 @@ void FRemoteControlWebInterfaceModule::OnSettingsModified(UObject* Settings, FPr
 		WebApp->Shutdown();
 		WebApp->Start();
 	}
+
+	if (PropertyName == GET_MEMBER_NAME_CHECKED(URemoteControlSettings, bWebAppLogRequestDuration))
+	{
+		const URemoteControlSettings* RCSettings = CastChecked<URemoteControlSettings>(Settings);
+		WebApp->SetExternalLoggerEnabled(RCSettings->bWebAppLogRequestDuration);
+	}
 }
 
 #undef LOCTEXT_NAMESPACE
diff --git a/Engine/Plugins/VirtualProduction/RemoteControlWebInterface/Source/RemoteControlWebInterface/Private/RCWebInterfaceProcess.cpp b/Engine/Plugins/VirtualProduction/RemoteControlWebInterface/Source/RemoteControlWebInterface/Private/RCWebInterfaceProcess.cpp
--- a/Engine/Plugins/VirtualProduction/RemoteControlWebInterface/Source/RemoteControlWebInterface/Private/RCWebInterfaceProcess.cpp
+++ b/Engine/Plugins/VirtualProduction/RemoteControlWebInterface/Source/RemoteControlWebInterface/Private/RCWebInterfaceProcess.cpp
@@ -8,6 +8,11 @@
 #include "GenericPlatform/GenericPlatformProcess.h"
 #include "HAL/PlatformFileManager.h"
 #include "Misc/CommandLine.h"
+#include "Modules/ModuleManager.h"
+#include "IWebRemoteControlModule.h"
+#include "IWebSocketNetworkingModule.h"
+#include "INetworkingWebSocket.h"
+#include "SocketSubsystem.h"
 
 
 #define LOCTEXT_NAMESPACE "RemoteControlWebInterface"
@@ -55,6 +60,8 @@ void FRemoteControlWebInterfaceProcess::Shutdown()
 {
 	Status = EStatus::Stopped;
 
+	SetExternalLoggerEnabled(false);
+
 	if (Process.IsValid())
 	{
 		FPlatformProcess::TerminateProc(Process, true);
@@ -74,6 +81,25 @@ FRemoteControlWebInterfaceProcess::EStatus FRemoteControlWebInterfaceProcess::Ge
 	return Status.load();
 }
 
+void FRemoteControlWebInterfaceProcess::SetExternalLoggerEnabled(bool bEnableExternalLog)
+{
+	TSharedPtr<INetworkingWebSocket> WebSocketLoggerConnection;
+	if (bEnableExternalLog)
+	{
+		const URemoteControlSettings* RCSettings = GetDefault<URemoteControlSettings>();
+
+		TSharedRef<FInternetAddr> Address = ISocketSubsystem::Get()->CreateInternetAddr();
+
+		bool bIsValidIp = false;
+		Address->SetIp(TEXT("127.0.0.1"), bIsValidIp);
+		Address->SetPort(RCSettings->RemoteControlWebInterfacePort + 2);
+		WebSocketLoggerConnection = FModuleManager::LoadModuleChecked<IWebSocketNetworkingModule>(TEXT("WebSocketNetworking")).CreateConnection(*Address);
+	}
+
+	IWebRemoteControlModule& WebRemoteControlModule = FModuleManager::LoadModuleChecked<IWebRemoteControlModule>("WebRemoteControl");
+	WebRemoteControlModule.SetExternalRemoteWebSocketLoggerConnection(WebSocketLoggerConnection);
+}
+
 uint32 FRemoteControlWebInterfaceProcess::Run()
 {
 	IPlatformFile& PlatformFile = FPlatformFileManager::Get().GetPlatformFile();
@@ -111,6 +137,11 @@ uint32 FRemoteControlWebInterfaceProcess::Run()
 		Args.Append(TEXT("--build "));
 	}
 
+	if (RCSettings->bWebAppLogRequestDuration)
+	{
+		Args.Append(TEXT("--log "));
+	}
+
 	void* ReadPipe = nullptr;
 	void* WritePipe = nullptr;
 
@@ -188,6 +219,12 @@ uint32 FRemoteControlWebInterfaceProcess::Run()
 							FText::FromString(CompleteMessage),
 							true);
 						Status = EStatus::Running;
+
+
+						if (RCSettings->bWebAppLogRequestDuration)
+						{
+							SetExternalLoggerEnabled(true);
+						}
 					}
 					else
 					{
diff --git a/Engine/Plugins/VirtualProduction/RemoteControlWebInterface/Source/RemoteControlWebInterface/Public/RCWebInterfaceProcess.h b/Engine/Plugins/VirtualProduction/RemoteControlWebInterface/Source/RemoteControlWebInterface/Public/RCWebInterfaceProcess.h
--- a/Engine/Plugins/VirtualProduction/RemoteControlWebInterface/Source/RemoteControlWebInterface/Public/RCWebInterfaceProcess.h
+++ b/Engine/Plugins/VirtualProduction/RemoteControlWebInterface/Source/RemoteControlWebInterface/Public/RCWebInterfaceProcess.h
@@ -34,6 +34,9 @@ public:
 	/** Get the webapp's status. */
 	EStatus GetStatus() const;
 
+	/** Enable / disable external logger */
+	void SetExternalLoggerEnabled(bool bEnableExternalLog);
+
 protected:
 	virtual uint32 Run() override;
 
diff --git a/Engine/Plugins/VirtualProduction/RemoteControlWebInterface/Source/RemoteControlWebInterface/RemoteControlWebInterface.Build.cs b/Engine/Plugins/VirtualProduction/RemoteControlWebInterface/Source/RemoteControlWebInterface/RemoteControlWebInterface.Build.cs
--- a/Engine/Plugins/VirtualProduction/RemoteControlWebInterface/Source/RemoteControlWebInterface/RemoteControlWebInterface.Build.cs
+++ b/Engine/Plugins/VirtualProduction/RemoteControlWebInterface/Source/RemoteControlWebInterface/RemoteControlWebInterface.Build.cs
@@ -18,9 +18,11 @@ public class RemoteControlWebInterface : ModuleRules
 				"SlateCore",
 				"Projects",
 				"RemoteControl",
-				"WebRemoteControl"
+				"Sockets",
+				"WebRemoteControl",
+				"WebSocketNetworking"
 			}
-			);
+		);
 
 		PrivateDependencyModuleNames.AddRange(
 			new string[] 
diff --git a/Engine/Plugins/VirtualProduction/RemoteControlWebInterface/WebApp/Server/src/App.ts b/Engine/Plugins/VirtualProduction/RemoteControlWebInterface/WebApp/Server/src/App.ts
--- a/Engine/Plugins/VirtualProduction/RemoteControlWebInterface/WebApp/Server/src/App.ts
+++ b/Engine/Plugins/VirtualProduction/RemoteControlWebInterface/WebApp/Server/src/App.ts
@@ -3,8 +3,7 @@ import { Request, Response, NextFunction, static as expressStatic } from 'expres
 import bodyParser from 'body-parser';
 import path from 'path';
 import request from 'superagent';
-import { Program, Api, Notify } from './';
-import { UnrealEngine } from './UnrealEngine';
+import { Program, Api, Notify, UnrealEngine, LogServer } from './';
 
 
 export class App extends Server {
@@ -46,6 +45,7 @@ export class App extends Server {
       const server = this.app.listen(Program.port, async () => {
         await Notify.initialize(server);
         await UnrealEngine.initialize();
+        await LogServer.initialize();
         resolve();
       });
 
diff --git a/Engine/Plugins/VirtualProduction/RemoteControlWebInterface/WebApp/Server/src/LogServer.ts b/Engine/Plugins/VirtualProduction/RemoteControlWebInterface/WebApp/Server/src/LogServer.ts
new file mode 100644
--- /dev/null
+++ b/Engine/Plugins/VirtualProduction/RemoteControlWebInterface/WebApp/Server/src/LogServer.ts
@@ -0,0 +1,69 @@
+import ws from 'ws';
+import _ from 'lodash';
+import fs from 'fs-extra';
+import { Program } from './';
+
+
+export namespace LogServer {
+
+  var server: ws.Server;
+  var buffer: string = '';
+  var startTime = new Date().toISOString().replace(/\:/g, '_');
+  var flushTimer: NodeJS.Timeout;
+
+  export async function initialize(): Promise<void> {
+    const port = Program.dev ? 7002 : Program.port + 2;
+    server = new ws.Server({ port });
+    await fs.ensureDir('./Logs');
+    server.on('connection', (socket: ws) => {
+      socket.on('message', onMessage);
+    });
+
+    flushTimer = setInterval(flusLog, 2 * 60 * 1000);
+  }
+
+  export async function shutdown(): Promise<void> {
+    server?.removeAllListeners();
+    server?.close();
+    server = null;
+    clearInterval(flushTimer);
+    flushTimer = null;
+  }
+
+  function onMessage(data: ws.Data) {
+    const json = data.toString();
+    try {
+      const payload = JSON.parse(json);
+      log(payload);
+    } catch (err) {
+      console.log("Error with message:", json);
+      console.log(err);
+    }
+  }
+
+  export function log(payload: Record<string, any>): void {
+    const Timestamp = new Date().toISOString();
+    buffer += JSON.stringify({ Timestamp, ...payload }) + '\r\n';
+  }
+  
+  function flusLog() {
+    if (flushLogFile(buffer, 'all'))
+      buffer = '';
+  }
+
+  function flushLogFile(content: string, type: string) {
+    if (!content)
+      return false;
+
+    try {
+      const filename = `./Logs/${startTime}-${type}.log`;
+      fs.appendFileSync(filename, content, 'utf8');
+
+      return true;
+    } catch (err) {
+      console.log('error', err);
+    }
+
+    return false;
+  }
+}
diff --git a/Engine/Plugins/VirtualProduction/RemoteControlWebInterface/WebApp/Server/src/Program.ts b/Engine/Plugins/VirtualProduction/RemoteControlWebInterface/WebApp/Server/src/Program.ts
--- a/Engine/Plugins/VirtualProduction/RemoteControlWebInterface/WebApp/Server/src/Program.ts
+++ b/Engine/Plugins/VirtualProduction/RemoteControlWebInterface/WebApp/Server/src/Program.ts
@@ -9,6 +9,7 @@ class Program {
   readonly ueHttpPort: number;
   readonly ueWebSocketPort: number;
   readonly rootFolder: string;
+  readonly logger: boolean;
 
   constructor() {
     this.dev = false;
@@ -16,12 +17,14 @@ class Program {
     this.port = 7000;
     this.ueHttpPort = 30010;
     this.ueWebSocketPort = 30020;
+    this.logger = false;
 
     for (let i = 2; i < process.argv.length; i++) {
       switch (process.argv[i]) {
         case '--dev':
           this.dev = true;
           this.port = 7001;
+          this.logger = true;
           break;
 
         case '--monitor':
@@ -42,6 +45,10 @@ class Program {
           this.ueWebSocketPort = parseInt(process.argv[i + 1]) || 30020;
           i++;
           break;
+
+        case '--log':
+          this.logger = true;
+          break;
       }
     }
 
diff --git a/Engine/Plugins/VirtualProduction/RemoteControlWebInterface/WebApp/Server/src/UnrealEngine.ts b/Engine/Plugins/VirtualProduction/RemoteControlWebInterface/WebApp/Server/src/UnrealEngine.ts
--- a/Engine/Plugins/VirtualProduction/RemoteControlWebInterface/WebApp/Server/src/UnrealEngine.ts
+++ b/Engine/Plugins/VirtualProduction/RemoteControlWebInterface/WebApp/Server/src/UnrealEngine.ts
@@ -2,7 +2,7 @@ import { IPayload, IPayloads, IPreset, IPresets, IPanel, IView, ICustomStackWidg
           IAsset, WidgetTypes, PropertyType, IColorPickerList, ICustomStackItem } from '../../Client/src/shared';
 import _ from 'lodash';
 import WebSocket from 'ws';
-import { Notify, Program } from './';
+import { Notify, Program, LogServer } from './';
 import crypto from 'crypto';
 
 
@@ -180,6 +180,9 @@ export namespace UnrealEngine {
         if (!promise)
           return;
 
+        if (Program.logger)
+          LogServer.log({ RequestId: message.RequestId, Stage: 'Done' });
+
         delete pendings[message.RequestId];
         promise?.(message.ResponseBody);
         return;
@@ -300,7 +303,11 @@ export namespace UnrealEngine {
 
   function http<T>(Verb: string, URL: string, Body?: object, wantAnswer?: boolean): Promise<T> {
     const RequestId = httpRequest++;
-    send('http', { RequestId, Verb, URL, Body });
+    const payload = { RequestId, Verb, URL, Body };
+    if (Program.logger)
+      LogServer.log(payload);
+
+    send('http', payload);
     if (!wantAnswer)
       return Promise.resolve(null);
 
diff --git a/Engine/Plugins/VirtualProduction/RemoteControlWebInterface/WebApp/Server/src/index.ts b/Engine/Plugins/VirtualProduction/RemoteControlWebInterface/WebApp/Server/src/index.ts
--- a/Engine/Plugins/VirtualProduction/RemoteControlWebInterface/WebApp/Server/src/index.ts
+++ b/Engine/Plugins/VirtualProduction/RemoteControlWebInterface/WebApp/Server/src/index.ts
@@ -1,6 +1,7 @@
 
 export * from './Api';
 export * from './App';
+export * from './LogServer';
 export * from './Notify';
 export * from './Program';
 export * from './UnrealEngine';
diff --git a/Engine/Plugins/VirtualProduction/RemoteControlWebInterface/WebApp/_package.json b/Engine/Plugins/VirtualProduction/RemoteControlWebInterface/WebApp/_package.json
--- a/Engine/Plugins/VirtualProduction/RemoteControlWebInterface/WebApp/_package.json
+++ b/Engine/Plugins/VirtualProduction/RemoteControlWebInterface/WebApp/_package.json
@@ -1,6 +1,6 @@
 {
   "name": "conductor",
-  "version": "0.3.56",
+  "version": "0.3.57",
   "scripts": {
     "postinstall": "concurrently \"cd ./Client && npm install && cd ..\" \"cd ./Server && npm install && cd ..\"",
     "build": "node ./Scripts/build.js",
