- Fixing crash at shutdown with remote control

diff --git a/Engine/Plugins/VirtualProduction/RemoteControlWebInterface/Source/RemoteControlWebInterface/Private/RCWebInterfaceProcess.cpp b/Engine/Plugins/VirtualProduction/RemoteControlWebInterface/Source/RemoteControlWebInterface/Private/RCWebInterfaceProcess.cpp
--- a/Engine/Plugins/VirtualProduction/RemoteControlWebInterface/Source/RemoteControlWebInterface/Private/RCWebInterfaceProcess.cpp
+++ b/Engine/Plugins/VirtualProduction/RemoteControlWebInterface/Source/RemoteControlWebInterface/Private/RCWebInterfaceProcess.cpp
@@ -96,8 +96,12 @@ void FRemoteControlWebInterfaceProcess::SetExternalLoggerEnabled(bool bEnableExt
 		WebSocketLoggerConnection = FModuleManager::LoadModuleChecked<IWebSocketNetworkingModule>(TEXT("WebSocketNetworking")).CreateConnection(*Address);
 	}
 
-	IWebRemoteControlModule& WebRemoteControlModule = FModuleManager::LoadModuleChecked<IWebRemoteControlModule>("WebRemoteControl");
-	WebRemoteControlModule.SetExternalRemoteWebSocketLoggerConnection(WebSocketLoggerConnection);
+
+	IWebRemoteControlModule* WebRemoteControlModule = FModuleManager::GetModulePtr<IWebRemoteControlModule>("WebRemoteControl");
+	if (WebRemoteControlModule)
+	{
+		WebRemoteControlModule->SetExternalRemoteWebSocketLoggerConnection(WebSocketLoggerConnection);
+	}
 }
 
 uint32 FRemoteControlWebInterfaceProcess::Run()
diff --git a/Engine/Source/Editor/PropertyEditor/Private/UserInterface/PropertyEditor/SPropertyEditorInteractiveActorPicker.cpp b/Engine/Source/Editor/PropertyEditor/Private/UserInterface/PropertyEditor/SPropertyEditorInteractiveActorPicker.cpp
--- a/Engine/Source/Editor/PropertyEditor/Private/UserInterface/PropertyEditor/SPropertyEditorInteractiveActorPicker.cpp
+++ b/Engine/Source/Editor/PropertyEditor/Private/UserInterface/PropertyEditor/SPropertyEditorInteractiveActorPicker.cpp
@@ -8,12 +8,13 @@
 
 SPropertyEditorInteractiveActorPicker::~SPropertyEditorInteractiveActorPicker()
 {
-	FActorPickerModeModule& ActorPickerMode = FModuleManager::Get().GetModuleChecked<FActorPickerModeModule>("ActorPickerMode");
-
-	// make sure we are unregistered when this widget goes away
-	if (ActorPickerMode.IsInActorPickingMode())
+	if (FActorPickerModeModule* ActorPickerMode = FModuleManager::Get().GetModulePtr<FActorPickerModeModule>("ActorPickerMode"))
 	{
-		ActorPickerMode.EndActorPickingMode();
+		// make sure we are unregistered when this widget goes away
+		if (ActorPickerMode->IsInActorPickingMode())
+		{
+			ActorPickerMode->EndActorPickingMode();
+		}
 	}
 }
 
diff --git a/Engine/Source/Editor/PropertyEditor/Private/UserInterface/PropertyEditor/SPropertyEditorSceneDepthPicker.cpp b/Engine/Source/Editor/PropertyEditor/Private/UserInterface/PropertyEditor/SPropertyEditorSceneDepthPicker.cpp
--- a/Engine/Source/Editor/PropertyEditor/Private/UserInterface/PropertyEditor/SPropertyEditorSceneDepthPicker.cpp
+++ b/Engine/Source/Editor/PropertyEditor/Private/UserInterface/PropertyEditor/SPropertyEditorSceneDepthPicker.cpp
@@ -9,10 +9,11 @@
 
 SPropertyEditorSceneDepthPicker::~SPropertyEditorSceneDepthPicker()
 {
-	FActorPickerModeModule& ActorPickerMode = FModuleManager::Get().GetModuleChecked<FActorPickerModeModule>("ActorPickerMode");
-
-	// make sure we are unregistered when this widget goes away
-	ActorPickerMode.EndActorPickingMode();
+	if (FActorPickerModeModule* ActorPickerMode = FModuleManager::Get().GetModulePtr<FActorPickerModeModule>("ActorPickerMode"))
+	{
+		// make sure we are unregistered when this widget goes away
+		ActorPickerMode->EndActorPickingMode();
+	}
 }
 
 void SPropertyEditorSceneDepthPicker::Construct(const FArguments& InArgs)
