- Fixing crash when launching RemoteControlWebapp

diff --git a/Engine/Plugins/VirtualProduction/RemoteControlWebInterface/Source/RemoteControlWebInterface/Private/RCWebInterfaceProcess.cpp b/Engine/Plugins/VirtualProduction/RemoteControlWebInterface/Source/RemoteControlWebInterface/Private/RCWebInterfaceProcess.cpp
--- a/Engine/Plugins/VirtualProduction/RemoteControlWebInterface/Source/RemoteControlWebInterface/Private/RCWebInterfaceProcess.cpp
+++ b/Engine/Plugins/VirtualProduction/RemoteControlWebInterface/Source/RemoteControlWebInterface/Private/RCWebInterfaceProcess.cpp
@@ -109,11 +109,9 @@ uint32 FRemoteControlWebInterfaceProcess::Run()
 	IPlatformFile& PlatformFile = FPlatformFileManager::Get().GetPlatformFile();
 
 #if PLATFORM_WINDOWS
-	const FString StartScript = Root / TEXT("Start.bat");
-#elif PLATFORM_MAC
-	const FString StartScript = Root / TEXT("Start.command");
+	FString StartScript = Root / TEXT("Start.bat");
 #else
-	const FString StartScript = Root / TEXT("Start.sh");
+	FString StartScript = Root / TEXT("Start.sh");
 #endif
 
 	FText ErrorTitle = LOCTEXT("RemoteControlWebInterface_ErrorTitle", "Failed to Launch the Remote Control Web Interface");
@@ -162,6 +160,11 @@ uint32 FRemoteControlWebInterfaceProcess::Run()
 
 	check(ReadPipe);
 	check(WritePipe);
+	
+#if PLATFORM_MAC
+	Args = FString::Printf(TEXT("-l %s %s"), *StartScript, *Args);
+	StartScript = TEXT("/bin/sh");
+#endif // PLATFORM_MAC
 
 	Process = FPlatformProcess::CreateProc(
 		*StartScript,	/* Path to start script */
@@ -281,4 +284,4 @@ uint32 FRemoteControlWebInterfaceProcess::Run()
 	return 0;
 }
 
-#undef LOCTEXT_NAMESPACE
\ No newline at end of file
+#undef LOCTEXT_NAMESPACE
diff --git a/Engine/Plugins/VirtualProduction/RemoteControlWebInterface/WebApp/Start.command b/Engine/Plugins/VirtualProduction/RemoteControlWebInterface/WebApp/Start.command
deleted file mode 100644
