RemoteControlWebInterface: Add CVar to start / stop webapp
Add missing widget types in dropdown

diff --git a/Engine/Plugins/VirtualProduction/RemoteControlWebInterface/Source/RemoteControlWebInterface/Private/RCWebInterface.cpp b/Engine/Plugins/VirtualProduction/RemoteControlWebInterface/Source/RemoteControlWebInterface/Private/RCWebInterface.cpp
--- a/Engine/Plugins/VirtualProduction/RemoteControlWebInterface/Source/RemoteControlWebInterface/Private/RCWebInterface.cpp
+++ b/Engine/Plugins/VirtualProduction/RemoteControlWebInterface/Source/RemoteControlWebInterface/Private/RCWebInterface.cpp
@@ -17,6 +17,15 @@
 
 #define LOCTEXT_NAMESPACE "FRemoteControlWebInterfaceModule"
 
+
+static TAutoConsoleVariable<int32> CVarRCWebInterfaceRunWebApp(
+	TEXT("RemoteControlWebInterface.RunWebApp"),
+	1,
+	TEXT("Run / Stop Remote Control Web App"));
+
+static FAutoConsoleVariableSink CVarRCWebInterfaceRunWebAppSink(FConsoleCommandDelegate::CreateStatic(&FRemoteControlWebInterfaceModule::OnCVarChanged));
+
+
 namespace RCWebInterface
 {
 	bool IsWebInterfaceEnabled()
@@ -33,14 +42,14 @@ namespace RCWebInterface
 
 void FRemoteControlWebInterfaceModule::StartupModule()
 {
+	WebApp = MakeShared<FRemoteControlWebInterfaceProcess>();
+
 	if (!RCWebInterface::IsWebInterfaceEnabled())
 	{
 		UE_LOG(LogRemoteControlWebInterface, Display, TEXT("Remote Control Web Interface is disabled by default when running outside the editor. Use the -RCWebInterfaceEnable flag when launching in order to use it."));
 		return;
 	}
 	
-	WebApp = MakeShared<FRemoteControlWebInterfaceProcess>();
-
 	TSharedPtr<FRemoteControlWebInterfaceProcess> WebAppLocal = WebApp;
 
 	FCoreDelegates::OnFEngineLoopInitComplete.AddLambda([WebAppLocal]()
@@ -88,6 +97,30 @@ void FRemoteControlWebInterfaceModule::ShutdownModule()
 	}
 }
 
+void FRemoteControlWebInterfaceModule::OnCVarChanged()
+{
+	FRemoteControlWebInterfaceModule& RCWebInf = FRemoteControlWebInterfaceModule::Get();
+	FRemoteControlWebInterfaceProcess::EStatus WebAppStatus = RCWebInf.WebApp->GetStatus();
+
+	bool bIsCurrentlyRunning = (
+		WebAppStatus == FRemoteControlWebInterfaceProcess::EStatus::Running ||
+		WebAppStatus == FRemoteControlWebInterfaceProcess::EStatus::Launching
+		);
+
+	bool bRunWebApp = CVarRCWebInterfaceRunWebApp.GetValueOnGameThread() != 0;
+	if (bRunWebApp != bIsCurrentlyRunning)
+	{
+		if (bRunWebApp)
+		{
+			RCWebInf.WebApp->Start();
+		}
+		else
+		{
+			RCWebInf.WebApp->Shutdown();
+		}
+	}
+}
+
 void FRemoteControlWebInterfaceModule::OnSettingsModified(UObject* Settings, FPropertyChangedEvent& PropertyChangedEvent)
 {
 	FName PropertyName = PropertyChangedEvent.GetPropertyName();
diff --git a/Engine/Plugins/VirtualProduction/RemoteControlWebInterface/Source/RemoteControlWebInterface/Private/RCWebInterfaceCustomizations.cpp b/Engine/Plugins/VirtualProduction/RemoteControlWebInterface/Source/RemoteControlWebInterface/Private/RCWebInterfaceCustomizations.cpp
--- a/Engine/Plugins/VirtualProduction/RemoteControlWebInterface/Source/RemoteControlWebInterface/Private/RCWebInterfaceCustomizations.cpp
+++ b/Engine/Plugins/VirtualProduction/RemoteControlWebInterface/Source/RemoteControlWebInterface/Private/RCWebInterfaceCustomizations.cpp
@@ -79,15 +79,15 @@ namespace RCWebInterface
 			{
 				if (Struct->IsChildOf(TBaseStructure<FVector>::Get()))
 				{
-					Widgets = { TEXT("Vector"), TEXT("Joystick"), TEXT("Sliders") };
+					Widgets = { TEXT("Vector"), TEXT("Joystick"), TEXT("Sliders"), TEXT("Dials") };
 				}
 				else if (Struct->IsChildOf(TBaseStructure<FVector2D>::Get()))
 				{
-					Widgets = { TEXT("Vector"), TEXT("Joystick"), TEXT("Sliders") };
+					Widgets = { TEXT("Vector"), TEXT("Joystick"), TEXT("Sliders"), TEXT("Dials") };
 				}
 				else if (Struct->IsChildOf(TBaseStructure<FRotator>::Get()))
 				{
-					Widgets = { TEXT("Vector"), TEXT("Sliders") };
+					Widgets = { TEXT("Vector"), TEXT("Sliders"), TEXT("Dials") };
 				}
 				else if (Struct->IsChildOf(TBaseStructure<FVector4>::Get()) || Struct->IsChildOf(TBaseStructure<FColor>::Get()) || Struct->IsChildOf(TBaseStructure<FLinearColor>::Get()))
 				{
@@ -95,6 +95,10 @@ namespace RCWebInterface
 				}
 			}
 		}
+		else if (Property->IsA<FObjectPropertyBase>())
+		{
+			Widgets = { TEXT("Asset") };
+		}
 
 		return Widgets;
 	}
diff --git a/Engine/Plugins/VirtualProduction/RemoteControlWebInterface/Source/RemoteControlWebInterface/Public/RCWebInterface.h b/Engine/Plugins/VirtualProduction/RemoteControlWebInterface/Source/RemoteControlWebInterface/Public/RCWebInterface.h
--- a/Engine/Plugins/VirtualProduction/RemoteControlWebInterface/Source/RemoteControlWebInterface/Public/RCWebInterface.h
+++ b/Engine/Plugins/VirtualProduction/RemoteControlWebInterface/Source/RemoteControlWebInterface/Public/RCWebInterface.h
@@ -17,6 +17,21 @@ public:
 	virtual void StartupModule() override;
 	virtual void ShutdownModule() override;
 
+
+	/**
+	 * Singleton-like access to this module's interface.  This is just for convenience!
+	 * Beware of calling this during the shutdown phase, though.  Your module might have been unloaded already.
+	 *
+	 * @return Returns singleton instance, loading the module on demand if needed
+	 */
+	static FRemoteControlWebInterfaceModule& Get()
+	{
+		static const FName ModuleName = TEXT("RemoteControlWebInterface");
+		return FModuleManager::LoadModuleChecked<FRemoteControlWebInterfaceModule>(ModuleName);
+	}
+
+	static void OnCVarChanged();
+
 private:
 	/** Handle Web Interface settings modifications. */
 	void OnSettingsModified(UObject* Settings, struct FPropertyChangedEvent& PropertyChangedEvent);
