RemoteControlWebInterface: Add Auto start cvar

diff --git a/Engine/Plugins/VirtualProduction/RemoteControlWebInterface/Source/RemoteControlWebInterface/Private/RCWebInterface.cpp b/Engine/Plugins/VirtualProduction/RemoteControlWebInterface/Source/RemoteControlWebInterface/Private/RCWebInterface.cpp
--- a/Engine/Plugins/VirtualProduction/RemoteControlWebInterface/Source/RemoteControlWebInterface/Private/RCWebInterface.cpp
+++ b/Engine/Plugins/VirtualProduction/RemoteControlWebInterface/Source/RemoteControlWebInterface/Private/RCWebInterface.cpp
@@ -17,14 +17,10 @@
 
 #define LOCTEXT_NAMESPACE "FRemoteControlWebInterfaceModule"
 
-
-static TAutoConsoleVariable<int32> CVarRCWebInterfaceRunWebApp(
-	TEXT("RemoteControlWebInterface.RunWebApp"),
+static TAutoConsoleVariable<int32> CVarRCWebInterfaceAutoStart(
+	TEXT("RCWebInterface.AutoStart"),
 	1,
-	TEXT("Run / Stop Remote Control Web App"));
-
-static FAutoConsoleVariableSink CVarRCWebInterfaceRunWebAppSink(FConsoleCommandDelegate::CreateStatic(&FRemoteControlWebInterfaceModule::OnCVarChanged));
-
+	TEXT("Auto Start Remote Control Web App"));
 
 namespace RCWebInterface
 {
@@ -49,6 +45,12 @@ void FRemoteControlWebInterfaceModule::StartupModule()
 		UE_LOG(LogRemoteControlWebInterface, Display, TEXT("Remote Control Web Interface is disabled by default when running outside the editor. Use the -RCWebInterfaceEnable flag when launching in order to use it."));
 		return;
 	}
+
+	if (CVarRCWebInterfaceAutoStart.GetValueOnGameThread() == 0)
+	{
+		UE_LOG(LogRemoteControlWebInterface, Display, TEXT("Remote Control Web Interface did not launch WebApp because CVar RCWebInterface.AutoStart is set to 0"));
+		return;
+	}
 	
 	TSharedPtr<FRemoteControlWebInterfaceProcess> WebAppLocal = WebApp;
 
@@ -97,28 +99,49 @@ void FRemoteControlWebInterfaceModule::ShutdownModule()
 	}
 }
 
-void FRemoteControlWebInterfaceModule::OnCVarChanged()
+bool FRemoteControlWebInterfaceModule::Exec(UWorld* InWorld, const TCHAR* Cmd, FOutputDevice& Ar)
 {
+	if (!FParse::Command(&Cmd, TEXT("RCWebInterface")))
+	{
+		return false;
+	}
+
 	FRemoteControlWebInterfaceModule& RCWebInf = FRemoteControlWebInterfaceModule::Get();
-	FRemoteControlWebInterfaceProcess::EStatus WebAppStatus = RCWebInf.WebApp->GetStatus();
+	const FRemoteControlWebInterfaceProcess::EStatus WebAppStatus = RCWebInf.WebApp->GetStatus();
 
-	bool bIsCurrentlyRunning = (
+	const bool bIsCurrentlyRunning = (
 		WebAppStatus == FRemoteControlWebInterfaceProcess::EStatus::Running ||
 		WebAppStatus == FRemoteControlWebInterfaceProcess::EStatus::Launching
-		);
+	);
 
-	bool bRunWebApp = CVarRCWebInterfaceRunWebApp.GetValueOnGameThread() != 0;
-	if (bRunWebApp != bIsCurrentlyRunning)
+	if (FParse::Command(&Cmd, TEXT("Start")))
 	{
-		if (bRunWebApp)
+		if (!bIsCurrentlyRunning) // Ignore if already running
 		{
+			Ar.Log(TEXT("RCWebInterface: Starting WebApp"));
+
 			RCWebInf.WebApp->Start();
 		}
 		else
 		{
-			RCWebInf.WebApp->Shutdown();
+			Ar.Log(TEXT("RCWebInterface: WebApp was already running"));
 		}
 	}
+	else if (FParse::Command(&Cmd, TEXT("Stop")))
+	{
+		Ar.Log(TEXT("RCWebInterface: Stopping WebApp"));
+
+		RCWebInf.WebApp->Shutdown();
+	}
+	else if (FParse::Command(&Cmd, TEXT("Restart")))
+	{
+		Ar.Log(TEXT("RCWebInterface: Restarting WebApp"));
+
+		RCWebInf.WebApp->Shutdown();
+		RCWebInf.WebApp->Start();
+	}
+
+	return true;
 }
 
 void FRemoteControlWebInterfaceModule::OnSettingsModified(UObject* Settings, FPropertyChangedEvent& PropertyChangedEvent)
diff --git a/Engine/Plugins/VirtualProduction/RemoteControlWebInterface/Source/RemoteControlWebInterface/Public/RCWebInterface.h b/Engine/Plugins/VirtualProduction/RemoteControlWebInterface/Source/RemoteControlWebInterface/Public/RCWebInterface.h
--- a/Engine/Plugins/VirtualProduction/RemoteControlWebInterface/Source/RemoteControlWebInterface/Public/RCWebInterface.h
+++ b/Engine/Plugins/VirtualProduction/RemoteControlWebInterface/Source/RemoteControlWebInterface/Public/RCWebInterface.h
@@ -10,13 +10,18 @@ class FRemoteControlWebInterfaceProcess;
 
 class FRCWebInterfaceCustomizations;
 
-class FRemoteControlWebInterfaceModule : public IModuleInterface
+class FRemoteControlWebInterfaceModule 
+	: public IModuleInterface
+	, public FSelfRegisteringExec
 {
 public:
 	//~ IModuleInterface
 	virtual void StartupModule() override;
 	virtual void ShutdownModule() override;
 
+	//~ FSelfRegisteringExec interface
+	virtual bool Exec(UWorld* InWorld, const TCHAR* Cmd, FOutputDevice& Ar) override;
+
 
 	/**
 	 * Singleton-like access to this module's interface.  This is just for convenience!
@@ -30,8 +35,6 @@ public:
 		return FModuleManager::LoadModuleChecked<FRemoteControlWebInterfaceModule>(ModuleName);
 	}
 
-	static void OnCVarChanged();
-
 private:
 	/** Handle Web Interface settings modifications. */
 	void OnSettingsModified(UObject* Settings, struct FPropertyChangedEvent& PropertyChangedEvent);
