Disable remote control in qa virtualproduction tests

diff --git a/Engine/Plugins/VirtualProduction/RemoteControl/Source/WebRemoteControl/Private/WebRemoteControl.cpp b/Engine/Plugins/VirtualProduction/RemoteControl/Source/WebRemoteControl/Private/WebRemoteControl.cpp
--- a/Engine/Plugins/VirtualProduction/RemoteControl/Source/WebRemoteControl/Private/WebRemoteControl.cpp
+++ b/Engine/Plugins/VirtualProduction/RemoteControl/Source/WebRemoteControl/Private/WebRemoteControl.cpp
@@ -275,8 +275,8 @@ void FWebRemoteControlModule::StartupModule()
 		return;
 	}
 
-	// By default, disable web remote control in -game and packaged game.
-	if (!WebRemoteControl::IsWebControlEnabledInEditor())
+	// By default, disable web remote control in -game, packaged game and on build machines
+	if (!WebRemoteControl::IsWebControlEnabledInEditor() || GIsBuildMachine)
 	{
 		UE_LOG(LogRemoteControl, Display, TEXT("Web remote control is disabled by default when running outside the editor. Use the -RCWebControlEnable flag when launching in order to use it."));
 		return;
diff --git a/Engine/Plugins/VirtualProduction/RemoteControlWebInterface/Source/RemoteControlWebInterface/Private/RCWebInterface.cpp b/Engine/Plugins/VirtualProduction/RemoteControlWebInterface/Source/RemoteControlWebInterface/Private/RCWebInterface.cpp
--- a/Engine/Plugins/VirtualProduction/RemoteControlWebInterface/Source/RemoteControlWebInterface/Private/RCWebInterface.cpp
+++ b/Engine/Plugins/VirtualProduction/RemoteControlWebInterface/Source/RemoteControlWebInterface/Private/RCWebInterface.cpp
@@ -31,8 +31,9 @@ namespace RCWebInterface
 #if WITH_EDITOR
 		bIsEditor = GIsEditor;
 #endif
-		// By default, remote control web interface is disabled in -game and packaged game.
-		return (!IsRunningCommandlet() && bIsEditor) || FParse::Param(FCommandLine::Get(), TEXT("RCWebInterfaceEnable"));
+
+		// By default, remote control web interface is disabled in -game, packaged game, and on build machines.
+		return !GIsBuildMachine || (!IsRunningCommandlet() && bIsEditor) || FParse::Param(FCommandLine::Get(), TEXT("RCWebInterfaceEnable"));
 	}
 }
 
