Fix remote control enabled by default when running horde job

diff --git a/Engine/Plugins/VirtualProduction/RemoteControlWebInterface/Source/RemoteControlWebInterface/Private/RCWebInterface.cpp b/Engine/Plugins/VirtualProduction/RemoteControlWebInterface/Source/RemoteControlWebInterface/Private/RCWebInterface.cpp
--- a/Engine/Plugins/VirtualProduction/RemoteControlWebInterface/Source/RemoteControlWebInterface/Private/RCWebInterface.cpp
+++ b/Engine/Plugins/VirtualProduction/RemoteControlWebInterface/Source/RemoteControlWebInterface/Private/RCWebInterface.cpp
@@ -33,7 +33,7 @@ namespace RCWebInterface
 #endif
 
 		// By default, remote control web interface is disabled in -game, packaged game, and on build machines.
-		return !GIsBuildMachine || (!IsRunningCommandlet() && bIsEditor) || FParse::Param(FCommandLine::Get(), TEXT("RCWebInterfaceEnable"));
+		return (!GIsBuildMachine && !IsRunningCommandlet() && bIsEditor) || FParse::Param(FCommandLine::Get(), TEXT("RCWebInterfaceEnable"));
 	}
 }
 
