Add logging to rc web interface startup
Disable errors in rc interface temporarily until I figure out why the web app process still gets started on build machines

diff --git a/Engine/Plugins/VirtualProduction/RemoteControlWebInterface/Source/RemoteControlWebInterface/Private/RCWebInterface.cpp b/Engine/Plugins/VirtualProduction/RemoteControlWebInterface/Source/RemoteControlWebInterface/Private/RCWebInterface.cpp
--- a/Engine/Plugins/VirtualProduction/RemoteControlWebInterface/Source/RemoteControlWebInterface/Private/RCWebInterface.cpp
+++ b/Engine/Plugins/VirtualProduction/RemoteControlWebInterface/Source/RemoteControlWebInterface/Private/RCWebInterface.cpp
@@ -26,6 +26,15 @@ namespace RCWebInterface
 {
 	bool IsWebInterfaceEnabled()
 	{
+		const bool bIsRunningTest = FParse::Param(FCommandLine::Get(), TEXT("-gauntlet"));
+		if (bIsRunningTest)
+		{
+			// Temporary workaround to disable this module when running automated tests.
+			return false;
+		}
+
+		UE_LOG(LogRemoteControlWebInterface, Display, TEXT("Remote Control Web Interface startup, IsBuildMachine: %s IsRunningTest: %s"), GIsBuildMachine ? TEXT("Yes") : TEXT("No"), bIsRunningTest ? TEXT("Yes") : TEXT("No"));
+		
 		bool bIsEditor = false;
 
 #if WITH_EDITOR
@@ -33,7 +42,7 @@ namespace RCWebInterface
 #endif
 
 		// By default, remote control web interface is disabled in -game, packaged game, and on build machines.
-		return (!GIsBuildMachine && !IsRunningCommandlet() && bIsEditor) || FParse::Param(FCommandLine::Get(), TEXT("RCWebInterfaceEnable"));
+		return !GIsBuildMachine || (!IsRunningCommandlet() && bIsEditor) || FParse::Param(FCommandLine::Get(), TEXT("RCWebInterfaceEnable"));
 	}
 }
 
