Fix freeze when changing web rc server port

diff --git a/Engine/Plugins/VirtualProduction/RemoteControlWebInterface/Source/RemoteControlWebInterface/Private/RCWebInterface.cpp b/Engine/Plugins/VirtualProduction/RemoteControlWebInterface/Source/RemoteControlWebInterface/Private/RCWebInterface.cpp
--- a/Engine/Plugins/VirtualProduction/RemoteControlWebInterface/Source/RemoteControlWebInterface/Private/RCWebInterface.cpp
+++ b/Engine/Plugins/VirtualProduction/RemoteControlWebInterface/Source/RemoteControlWebInterface/Private/RCWebInterface.cpp
@@ -152,12 +152,17 @@ bool FRemoteControlWebInterfaceModule::Exec(UWorld* InWorld, const TCHAR* Cmd, F
 void FRemoteControlWebInterfaceModule::OnSettingsModified(UObject* Settings, FPropertyChangedEvent& PropertyChangedEvent)
 {
 	FName PropertyName = PropertyChangedEvent.GetPropertyName();
-	if (PropertyName == GET_MEMBER_NAME_CHECKED(URemoteControlSettings, RemoteControlWebInterfacePort) || PropertyName == GET_MEMBER_NAME_CHECKED(URemoteControlSettings, bForceWebAppBuildAtStartup))
+
+	if (PropertyChangedEvent.ChangeType != EPropertyChangeType::Interactive)
 	{
-		WebApp->Shutdown();
-		WebApp->Start();
+		if (PropertyName == GET_MEMBER_NAME_CHECKED(URemoteControlSettings, RemoteControlWebInterfacePort) || PropertyName == GET_MEMBER_NAME_CHECKED(URemoteControlSettings, bForceWebAppBuildAtStartup))
+		{
+			WebApp->Shutdown();
+			WebApp->Start();
+		}
 	}
 
+
 	if (PropertyName == GET_MEMBER_NAME_CHECKED(URemoteControlSettings, bWebAppLogRequestDuration))
 	{
 		const URemoteControlSettings* RCSettings = CastChecked<URemoteControlSettings>(Settings);
diff --git a/Engine/Plugins/VirtualProduction/RemoteControlWebInterface/Source/RemoteControlWebInterface/Private/RCWebInterfaceProcess.cpp b/Engine/Plugins/VirtualProduction/RemoteControlWebInterface/Source/RemoteControlWebInterface/Private/RCWebInterfaceProcess.cpp
--- a/Engine/Plugins/VirtualProduction/RemoteControlWebInterface/Source/RemoteControlWebInterface/Private/RCWebInterfaceProcess.cpp
+++ b/Engine/Plugins/VirtualProduction/RemoteControlWebInterface/Source/RemoteControlWebInterface/Private/RCWebInterfaceProcess.cpp
@@ -250,7 +250,7 @@ uint32 FRemoteControlWebInterfaceProcess::Run()
 
 	UE_LOG(LogRemoteControlWebInterface, Log, TEXT("WebApp started, initial launch will take longer as it will be building the WebApp"));
 
-	while (FPlatformProcess::IsProcRunning(Process))
+	while (Status != EStatus::Stopped && FPlatformProcess::IsProcRunning(Process))
 	{
 		const int32 BytesRead = LogReadPipe(ReadPipe);
 
