[RemoteControl] Fix Web App keep sending notification
Web App keep sending notification when WebSocketServer is disabled

diff --git a/Engine/Plugins/VirtualProduction/RemoteControlWebInterface/Source/RemoteControlWebInterface/Private/RCWebInterface.cpp b/Engine/Plugins/VirtualProduction/RemoteControlWebInterface/Source/RemoteControlWebInterface/Private/RCWebInterface.cpp
--- a/Engine/Plugins/VirtualProduction/RemoteControlWebInterface/Source/RemoteControlWebInterface/Private/RCWebInterface.cpp
+++ b/Engine/Plugins/VirtualProduction/RemoteControlWebInterface/Source/RemoteControlWebInterface/Private/RCWebInterface.cpp
@@ -61,7 +61,10 @@ void FRemoteControlWebInterfaceModule::StartupModule()
 
 	FCoreDelegates::OnFEngineLoopInitComplete.AddLambda([WebAppLocal]()
 		{
-			WebAppLocal->Start();
+			if (GetDefault<URemoteControlSettings>()->bAutoStartWebSocketServer)
+			{
+				WebAppLocal->Start();
+			}
 		});
 
 	if (IWebRemoteControlModule* WebRemoteControlModule = FModuleManager::GetModulePtr<IWebRemoteControlModule>("WebRemoteControl"))
