[Remote Control] WebInterface support for stateless controller

- Native Backend: Add the "Button" widget type for the the stateless controller in the widget metadata
- Webapp: adding the handling for stateless controller that has no value.

diff --git a/Engine/Plugins/VirtualProduction/RemoteControlWebInterface/Source/RemoteControlWebInterface/Private/RCWebInterfaceCustomizations.cpp b/Engine/Plugins/VirtualProduction/RemoteControlWebInterface/Source/RemoteControlWebInterface/Private/RCWebInterfaceCustomizations.cpp
--- a/Engine/Plugins/VirtualProduction/RemoteControlWebInterface/Source/RemoteControlWebInterface/Private/RCWebInterfaceCustomizations.cpp
+++ b/Engine/Plugins/VirtualProduction/RemoteControlWebInterface/Source/RemoteControlWebInterface/Private/RCWebInterfaceCustomizations.cpp
@@ -2,33 +2,33 @@
 
 #include "RCWebInterfaceCustomizations.h"
 
-#include "RCVirtualProperty.h"
-
 #if WITH_EDITOR
+#include "Controller/RCStatelessEventController.h"
 #include "DetailCategoryBuilder.h"
 #include "DetailLayoutBuilder.h"
 #include "DetailWidgetRow.h"
 #include "Editor.h"
-#include "Styling/AppStyle.h"
 #include "EditorFontGlyphs.h"
 #include "HAL/PlatformProcess.h"
-#include "Input/Reply.h"
 #include "IRemoteControlModule.h"
 #include "IRemoteControlUIModule.h"
 #include "ISettingsModule.h"
+#include "Input/Reply.h"
+#include "RCVirtualProperty.h"
 #include "RCWebInterfaceProcess.h"
-#include "RemoteControlSettings.h"
-#include "RemoteControlPreset.h"
 #include "RemoteControlEntity.h"
-#include "ScopedTransaction.h"
+#include "RemoteControlPreset.h"
+#include "RemoteControlSettings.h"
 #include "SSearchableComboBox.h"
+#include "ScopedTransaction.h"
+#include "Styling/AppStyle.h"
 #include "UObject/TextProperty.h"
+#include "Widgets/DeclarativeSyntaxSupport.h"
+#include "Widgets/Images/SThrobber.h"
 #include "Widgets/Input/SButton.h"
 #include "Widgets/Input/SEditableTextBox.h"
-#include "Widgets/Text/STextBlock.h"
-#include "Widgets/DeclarativeSyntaxSupport.h"
 #include "Widgets/Layout/SWidgetSwitcher.h"
-#include "Widgets/Images/SThrobber.h"
+#include "Widgets/Text/STextBlock.h"
 
 #define LOCTEXT_NAMESPACE "FRCWebInterfaceCustomizations"
 
@@ -505,7 +505,14 @@ FString FRCWebInterfaceCustomizations::OnInitializeWidgetMetadata(URemoteControl
 	}
 	else if (const URCVirtualPropertyBase* Controller = Preset->GetController(EntityId))
 	{
-		Metadata = RCWebInterface::GetDefaultWidget(const_cast<FProperty*>(Controller->GetProperty()));
+		if (FRCStatelessEventController::IsStatelessEventController(Controller))
+		{
+			Metadata = TEXT("Button");
+		}
+		else
+		{
+			Metadata = RCWebInterface::GetDefaultWidget(const_cast<FProperty*>(Controller->GetProperty()));
+		}
 	}
 	
 	return Metadata;
diff --git a/Engine/Plugins/VirtualProduction/RemoteControlWebInterface/Source/RemoteControlWebInterface/RemoteControlWebInterface.Build.cs b/Engine/Plugins/VirtualProduction/RemoteControlWebInterface/Source/RemoteControlWebInterface/RemoteControlWebInterface.Build.cs
--- a/Engine/Plugins/VirtualProduction/RemoteControlWebInterface/Source/RemoteControlWebInterface/RemoteControlWebInterface.Build.cs
+++ b/Engine/Plugins/VirtualProduction/RemoteControlWebInterface/Source/RemoteControlWebInterface/RemoteControlWebInterface.Build.cs
@@ -43,6 +43,7 @@ public class RemoteControlWebInterface : ModuleRules
 					"EditorWidgets",
 					"PropertyEditor",
 					"RemoteControl",
+					"RemoteControlLogic",
 					"RemoteControlUI",
 					"Settings",
 					"UnrealEd",
diff --git a/Engine/Plugins/VirtualProduction/RemoteControlWebInterface/WebApp/Client/src/components/Stack.tsx b/Engine/Plugins/VirtualProduction/RemoteControlWebInterface/WebApp/Client/src/components/Stack.tsx
--- a/Engine/Plugins/VirtualProduction/RemoteControlWebInterface/WebApp/Client/src/components/Stack.tsx
+++ b/Engine/Plugins/VirtualProduction/RemoteControlWebInterface/WebApp/Client/src/components/Stack.tsx
@@ -157,6 +157,9 @@ export class Stack extends React.Component<Props & PropsFromState, State> {
     if (widget.propertyType === PropertyType.Function)
       return _api.payload.execute(widget.property, widget.args);
 
+    if (widget.propertyType === PropertyType.StatelessController)
+      return _api.payload.set(widget.property, ""); // stateless controller have no value.
+
     if (value === undefined && !await AlertModal.show('Reset to default?'))
       return;
 
diff --git a/Engine/Plugins/VirtualProduction/RemoteControlWebInterface/WebApp/Client/src/shared/index.ts b/Engine/Plugins/VirtualProduction/RemoteControlWebInterface/WebApp/Client/src/shared/index.ts
--- a/Engine/Plugins/VirtualProduction/RemoteControlWebInterface/WebApp/Client/src/shared/index.ts
+++ b/Engine/Plugins/VirtualProduction/RemoteControlWebInterface/WebApp/Client/src/shared/index.ts
@@ -20,6 +20,7 @@ export enum PropertyType {
   Text =            'FText',
 
   Function =        'Function',
+  StatelessController = 'FRCStatelessEventController',
 }
 
 export enum ConnectionSignal {
