Remote Control Web Interface: fixed issue with web process shared ptr being captured by copy in a global multicast delegate.

diff --git a/Engine/Plugins/VirtualProduction/RemoteControlWebInterface/Source/RemoteControlWebInterface/Private/RCWebInterface.cpp b/Engine/Plugins/VirtualProduction/RemoteControlWebInterface/Source/RemoteControlWebInterface/Private/RCWebInterface.cpp
--- a/Engine/Plugins/VirtualProduction/RemoteControlWebInterface/Source/RemoteControlWebInterface/Private/RCWebInterface.cpp
+++ b/Engine/Plugins/VirtualProduction/RemoteControlWebInterface/Source/RemoteControlWebInterface/Private/RCWebInterface.cpp
@@ -56,14 +56,13 @@ void FRemoteControlWebInterfaceModule::StartupModule()
 		UE_LOG(LogRemoteControlWebInterface, Display, TEXT("Remote Control Web Interface did not launch WebApp because CVar RCWebInterface.AutoStart is set to 0"));
 		return;
 	}
-	
-	TSharedPtr<FRemoteControlWebInterfaceProcess> WebAppLocal = WebApp;
 
-	FCoreDelegates::OnFEngineLoopInitComplete.AddLambda([WebAppLocal]()
+	FCoreDelegates::OnFEngineLoopInitComplete.AddLambda([WebAppWeak = WebApp.ToWeakPtr()]()
 		{
-			if (GetDefault<URemoteControlSettings>()->bAutoStartWebSocketServer)
+			const TSharedPtr<FRemoteControlWebInterfaceProcess> WebApp = WebAppWeak.Pin();
+			if (WebApp.IsValid() && GetDefault<URemoteControlSettings>()->bAutoStartWebSocketServer)
 			{
-				WebAppLocal->Start();
+				WebApp->Start();
 			}
 		});
 
@@ -88,12 +87,13 @@ void FRemoteControlWebInterfaceModule::StartupModule()
 
 void FRemoteControlWebInterfaceModule::ShutdownModule()
 {
-	if (!RCWebInterface::IsWebInterfaceEnabled())
+	if (WebApp.IsValid())
 	{
-		return;
+		WebApp->Shutdown();
 	}
 
-	WebApp->Shutdown();
+	WebApp.Reset();
+	Customizations.Reset();
 
 	FCoreDelegates::OnFEngineLoopInitComplete.RemoveAll(this);
 
@@ -104,6 +104,7 @@ void FRemoteControlWebInterfaceModule::ShutdownModule()
 	if (IWebRemoteControlModule* WebRemoteControlModule = FModuleManager::GetModulePtr<IWebRemoteControlModule>("WebRemoteControl"))
 	{
 		WebRemoteControlModule->OnWebSocketServerStarted().Remove(WebSocketServerStartedDelegate);
+		WebSocketServerStartedDelegate.Reset();
 	}
 }
 
