Remote Control: fix issue with Mac OS failing to build remote control web app when the engine is installed in a directory containing spaces

diff --git a/Engine/Plugins/VirtualProduction/RemoteControlWebInterface/Source/RemoteControlWebInterface/Private/RCWebInterfaceProcess.cpp b/Engine/Plugins/VirtualProduction/RemoteControlWebInterface/Source/RemoteControlWebInterface/Private/RCWebInterfaceProcess.cpp
--- a/Engine/Plugins/VirtualProduction/RemoteControlWebInterface/Source/RemoteControlWebInterface/Private/RCWebInterfaceProcess.cpp
+++ b/Engine/Plugins/VirtualProduction/RemoteControlWebInterface/Source/RemoteControlWebInterface/Private/RCWebInterfaceProcess.cpp
@@ -161,9 +161,9 @@ uint32 FRemoteControlWebInterfaceProcess::Run()
 
 	check(ReadPipe);
 	check(WritePipe);
-	
+
 #if PLATFORM_MAC
-	Args = FString::Printf(TEXT("-l %s %s"), *StartScript, *Args);
+	Args = FString::Printf(TEXT("-l \"%s\" %s"), *StartScript, *Args);
 	StartScript = TEXT("/bin/sh");
 #endif // PLATFORM_MAC
 
